DMX - Improve DMXDisplayClusterReplicator ensure conditions

diff --git a/Engine/Plugins/VirtualProduction/DMX/DMXDisplayCluster/Source/DMXDisplayCluster/Private/DMXDisplayClusterReplicator.cpp b/Engine/Plugins/VirtualProduction/DMX/DMXDisplayCluster/Source/DMXDisplayCluster/Private/DMXDisplayClusterReplicator.cpp
--- a/Engine/Plugins/VirtualProduction/DMX/DMXDisplayCluster/Source/DMXDisplayCluster/Private/DMXDisplayClusterReplicator.cpp
+++ b/Engine/Plugins/VirtualProduction/DMX/DMXDisplayCluster/Source/DMXDisplayCluster/Private/DMXDisplayClusterReplicator.cpp
@@ -65,7 +65,10 @@ FDMXDisplayClusterReplicator::FDMXDisplayClusterReplicator()
 		bForceMaster = FParse::Param(FCommandLine::Get(), TEXT("dc_dmx_master"));
 		bForceSlave = FParse::Param(FCommandLine::Get(), TEXT("dc_dmx_slave"));
 
-		ensureMsgf(!(bForceMaster && bForceSlave), TEXT("Ambigous command line for the DMX display cluster plugin. An instance cannot be 'dc_dmx_master' and 'dc_dmx_slave' at the same time."));
+		if (!ensureMsgf(!(bForceMaster && bForceSlave), TEXT("Ambigous command line for the DMX display cluster plugin. An instance cannot be 'dc_dmx_master' and 'dc_dmx_slave' at the same time.")))
+		{
+			return;
+		}
 
 		// Bind to OnClusterEventReceived. The master needs to bind here too to get data at the same time as slaves. 
 		BinaryListener = FOnClusterEventBinaryListener::CreateRaw(this, &FDMXDisplayClusterReplicator::OnClusterEventReceived);
@@ -141,7 +144,7 @@ void FDMXDisplayClusterReplicator::OnClusterEventReceived(const FDisplayClusterC
 		FDMXDisplayClusterPacket DMXDisplayClusterPacket;
 		DMXDisplayClusterPacket.Serialize(Reader);
 
-		if (ensureAlwaysMsgf(CachedInputPorts.IsValidIndex(DMXDisplayClusterPacket.PortIndex), TEXT("FDMXDisplayClusterReplicator: Input Ports specified in project settings are not identical across the nDipslay cluster. DMX replication failed.")))
+		if (ensureMsgf(CachedInputPorts.IsValidIndex(DMXDisplayClusterPacket.PortIndex), TEXT("FDMXDisplayClusterReplicator: Input Ports specified in project settings are not identical across the nDipslay cluster. DMX replication failed.")))
 		{
 			CachedInputPorts[DMXDisplayClusterPacket.PortIndex]->GameThreadInjectDMXSignal(DMXDisplayClusterPacket.Signal.ToSharedRef());
 		}
