DMX - Avoid offensive language in FDMXDisplayClusterReplicator

diff --git a/Engine/Plugins/VirtualProduction/DMX/DMXDisplayCluster/Source/DMXDisplayCluster/Private/DMXDisplayClusterReplicator.cpp b/Engine/Plugins/VirtualProduction/DMX/DMXDisplayCluster/Source/DMXDisplayCluster/Private/DMXDisplayClusterReplicator.cpp
--- a/Engine/Plugins/VirtualProduction/DMX/DMXDisplayCluster/Source/DMXDisplayCluster/Private/DMXDisplayClusterReplicator.cpp
+++ b/Engine/Plugins/VirtualProduction/DMX/DMXDisplayCluster/Source/DMXDisplayCluster/Private/DMXDisplayClusterReplicator.cpp
@@ -53,9 +53,7 @@ namespace
 
 
 FDMXDisplayClusterReplicator::FDMXDisplayClusterReplicator()
-	: bForceMaster(false)
-	, bForceSlave(false)
-	, bClusterEventEmitter(false)
+	: bClusterEventEmitter(false)
 	, TickType(ETickableTickType::Never)
 {
 
@@ -64,38 +62,38 @@ FDMXDisplayClusterReplicator::FDMXDisplayClusterReplicator()
 
 	if (IDisplayClusterClusterManager* ClusterManager = IDisplayCluster::Get().GetClusterMgr())
 	{
-		bForceMaster = FParse::Param(FCommandLine::Get(), TEXT("dc_dmx_master"));
-		bForceSlave = FParse::Param(FCommandLine::Get(), TEXT("dc_dmx_slave"));
+		const bool bForcePrimary = FParse::Param(FCommandLine::Get(), TEXT("dc_dmx_primary"));
+		const bool bForceSecondary = FParse::Param(FCommandLine::Get(), TEXT("dc_dmx_secondary"));
 
-		if (!ensureMsgf(!(bForceMaster && bForceSlave), TEXT("Ambigous command line for the DMX display cluster plugin. An instance cannot be 'dc_dmx_master' and 'dc_dmx_slave' at the same time.")))
+		if (!ensureMsgf(!(bForcePrimary && bForceSecondary), TEXT("Ambigous command line for the DMX display cluster plugin. An instance cannot be 'dc_dmx_primary' and 'dc_dmx_secondary' at the same time.")))
 		{
 			return;
 		}
 
-		// Bind to OnClusterEventReceived. The master needs to bind here too to get data at the same time as slaves. 
+		// Bind to OnClusterEventReceived. The DMX primary node needs to bind here too to get data at the same time as secondary nodes. 
 		BinaryListener = FOnClusterEventBinaryListener::CreateRaw(this, &FDMXDisplayClusterReplicator::OnClusterEventReceived);
 		ClusterManager->AddClusterEventBinaryListener(BinaryListener);
 
 		// Cache the input ports. As the application runs in cluster mode without editor, the Input ports array is not expected to change.
 		CachedInputPorts = FDMXPortManager::Get().GetInputPorts();
 
-		bClusterEventEmitter = [ClusterManager, this]()
+		bClusterEventEmitter = [ClusterManager, bForcePrimary, bForceSecondary]()
 		{
-			const bool bAutoMode = !bForceMaster && !bForceSlave;
+			const bool bAutoMode = !bForcePrimary && !bForceSecondary;
 
 			if (bAutoMode)
 			{
 				return ClusterManager->IsPrimary();
 			}
 
-			return bForceMaster;
+			return bForcePrimary;
 		}();
 
 		if (bClusterEventEmitter)
 		{
 			TickType = ETickableTickType::Always;
 
-			// User a custom listener to avoid issues with tick precedence between the port and this object's tick here
+			// User a custom listener to avoid issues with tick precedence between the port and this object's tick 
 			for (const FDMXInputPortSharedRef& InputPort : CachedInputPorts)
 			{
 				// Don't use the default queue for the port, instead use a custom raw listener
@@ -113,7 +111,6 @@ FDMXDisplayClusterReplicator::FDMXDisplayClusterReplicator()
 			TickType = ETickableTickType::Never;
 
 			// If this is not an emitter, suspend receiving from protocols to not get any data from the network.
-			// This also allows to inject data into the ports - This is the single producer now.
 			FDMXPortManager::Get().SuspendProtocols();
 		}
 	}
@@ -188,8 +185,8 @@ void FDMXDisplayClusterReplicator::Tick(float DeltaTime)
 				ClusterEvent.bShouldDiscardOnRepeat = false;
 				ClusterEvent.EventId = DMXnDisplayReplicationEventID;
 				ClusterEvent.EventData = MoveTemp(ArrayWriter);
-				constexpr bool bEmitFromMasterOnly = false;
-				ClusterManager->EmitClusterEventBinary(ClusterEvent, bEmitFromMasterOnly);
+				constexpr bool bEmitFromPrimaryOnly = false;
+				ClusterManager->EmitClusterEventBinary(ClusterEvent, bEmitFromPrimaryOnly);
 			}
 		}
 	}
diff --git a/Engine/Plugins/VirtualProduction/DMX/DMXDisplayCluster/Source/DMXDisplayCluster/Private/DMXDisplayClusterReplicator.h b/Engine/Plugins/VirtualProduction/DMX/DMXDisplayCluster/Source/DMXDisplayCluster/Private/DMXDisplayClusterReplicator.h
--- a/Engine/Plugins/VirtualProduction/DMX/DMXDisplayCluster/Source/DMXDisplayCluster/Private/DMXDisplayClusterReplicator.h
+++ b/Engine/Plugins/VirtualProduction/DMX/DMXDisplayCluster/Source/DMXDisplayCluster/Private/DMXDisplayClusterReplicator.h
@@ -22,9 +22,6 @@ public:
 	virtual ~FDMXDisplayClusterReplicator();
 
 private:
-	/** Initalizes the replicator */
-	void Initialize();
-
 	/** Called whena cluster event was received */
 	void OnClusterEventReceived(const FDisplayClusterClusterEventBinary& Event);
 
@@ -35,12 +32,6 @@ private:
 	virtual TStatId GetStatId() const override;
 	//~End FTickableGameObject interface
 
-	/** Force the node to be a DMX Master */
-	bool bForceMaster;
-
-	/** Force the node to be a DMX Slave */
-	bool bForceSlave;
-
 	/** True if it is a cluster event emitter, else it is a cluster event receiver */
 	bool bClusterEventEmitter;
 
