DMX - Implement DMX controls for Light Cards

diff --git a/Engine/Plugins/VirtualProduction/DMX/DMXDisplayCluster/DMXDisplayCluster.uplugin b/Engine/Plugins/VirtualProduction/DMX/DMXDisplayCluster/DMXDisplayCluster.uplugin
--- a/Engine/Plugins/VirtualProduction/DMX/DMXDisplayCluster/DMXDisplayCluster.uplugin
+++ b/Engine/Plugins/VirtualProduction/DMX/DMXDisplayCluster/DMXDisplayCluster.uplugin
@@ -1,27 +1,37 @@
 {
-  "FileVersion": 1,
-  "Version": 1,
-  "VersionName": "1.0",
-  "FriendlyName": "DMX DisplayCluster",
-  "Description": "Allows integration between DMX and DisplayCluster",
-  "Category": "Virtual Production",
-  "CreatedBy": "Epic Games, Inc.",
-  "CreatedByURL": "http://epicgames.com",
-  "DocsURL": "",
-  "MarketplaceURL": "",
-  "SupportURL": "",
-  "CanContainContent": true,
-  "IsBetaVersion": true,
-  "Modules": [
-    {
-      "Name": "DMXDisplayCluster",
-      "Type": "Runtime",
-      "LoadingPhase": "Default",
-      "PlatformAllowList": [
-                                "Win64"
-                            ]
-    }
-   ],
+	"FileVersion": 1,
+	"Version": 1,
+	"VersionName": "1.0",
+	"FriendlyName": "DMX DisplayCluster",
+	"Description": "Allows integration between DMX and DisplayCluster",
+	"Category": "Virtual Production",
+	"CreatedBy": "Epic Games, Inc.",
+	"CreatedByURL": "http://epicgames.com",
+	"DocsURL": "",
+	"MarketplaceURL": "",
+	"SupportURL": "",
+	"CanContainContent": true,
+	"IsBetaVersion": true,
+	"Modules": [
+		{
+			"Name": "DMXDisplayCluster",
+			"Type": "Runtime",
+			"LoadingPhase": "Default",
+			"PlatformAllowList": [
+				"Win64",
+				"Linux"
+			]
+		},
+		{
+			"Name": "DMXDisplayClusterLightCard",
+			"Type": "Runtime",
+			"LoadingPhase": "Default",
+			"PlatformAllowList": [
+				"Win64",
+				"Linux"
+			]
+		}
+	],
 	"Plugins": [
 		{
 			"Name": "DMXEngine",
@@ -34,6 +44,10 @@
 		{
 			"Name": "nDisplay",
 			"Enabled": true
+		},
+		{
+			"Name": "nDisplayModularFeatures",
+			"Enabled": true
 		}
 	]
 }
\ No newline at end of file
diff --git a/Engine/Plugins/VirtualProduction/DMX/DMXDisplayCluster/Source/DMXDisplayClusterLightCard/DMXDisplayClusterLightCard.Build.cs b/Engine/Plugins/VirtualProduction/DMX/DMXDisplayCluster/Source/DMXDisplayClusterLightCard/DMXDisplayClusterLightCard.Build.cs
new file mode 100644
--- /dev/null
+++ b/Engine/Plugins/VirtualProduction/DMX/DMXDisplayCluster/Source/DMXDisplayClusterLightCard/DMXDisplayClusterLightCard.Build.cs
@@ -0,0 +1,26 @@
+// Copyright Epic Games, Inc. All Rights Reserved.
+
+using UnrealBuildTool;
+
+public class DMXDisplayClusterLightCard : ModuleRules
+{
+	public DMXDisplayClusterLightCard(ReadOnlyTargetRules Target) : base(Target)
+	{
+		PCHUsage = PCHUsageMode.UseExplicitOrSharedPCHs;
+				
+		PublicDependencyModuleNames.AddRange(new string[]
+		{                   
+			"DisplayCluster",
+			"DisplayClusterLightCardExtender",
+			"DMXProtocol",
+			"DMXRuntime",
+		});
+
+		PrivateDependencyModuleNames.AddRange(new string[]
+		{
+			"Core",
+			"CoreUObject",
+			"Engine",
+		});
+	}
+}
diff --git a/Engine/Plugins/VirtualProduction/DMX/DMXDisplayCluster/Source/DMXDisplayClusterLightCard/Private/DMXDisplayClusterLightCardComponent.cpp b/Engine/Plugins/VirtualProduction/DMX/DMXDisplayCluster/Source/DMXDisplayClusterLightCard/Private/DMXDisplayClusterLightCardComponent.cpp
new file mode 100644
--- /dev/null
+++ b/Engine/Plugins/VirtualProduction/DMX/DMXDisplayCluster/Source/DMXDisplayClusterLightCard/Private/DMXDisplayClusterLightCardComponent.cpp
@@ -0,0 +1,613 @@
+// Copyright Epic Games, Inc. All Rights Reserved.
+
+#include "DMXDisplayClusterLightCardComponent.h"
+
+#include "IO/DMXInputPort.h"
+#include "IO/DMXPortManager.h"
+#include "Library/DMXEntityFixturePatch.h"
+
+#include "Serialization/ArrayReader.h"
+
+
+/** Attribute names that match the current version of the DisplayClusterLightCard GDTF */
+namespace UE::DMX::DMXDisplayClusterLightCardComponent::Private
+{
+	namespace AttributeNames
+	{
+		const FName Attribute_DMXInput = "DMXInput";
+		const FName Attribute_DistanceFromCenter = "DistanceFromCenter";
+		const FName Attribute_Pan = "Pan";
+		const FName Attribute_Tilt = "Tilt";
+		const FName Attribute_RotX = "Spin";
+		const FName Attribute_RotY = "Pitch";
+		const FName Attribute_RotZ = "Yaw";
+		const FName Attribute_ScaleX = "ScaleX";
+		const FName Attribute_ScaleY = "ScaleY";
+		const FName Attribute_Mask = "Mask";
+		const FName Attribute_Red = "Red";
+		const FName Attribute_Green = "Green";
+		const FName Attribute_Blue = "Blue";
+		const FName Attribute_ColorAdd_Alpha = "ColorAdd_Alpha";
+		const FName Attribute_CTC = "CTC";
+		const FName Attribute_Tint = "Tint";
+		const FName Attribute_Exposure = "Exposure";
+		const FName Attribute_Gain = "Gain";
+		const FName Attribute_Opacity = "Opacity";
+		const FName Attribute_Feathering = "Feathering";
+		const FName Attribute_Alpha_Gradient_Enable = "Alpha_Gradient_Enable";
+		const FName Attribute_StartingAlpha = "StartingAlpha";
+		const FName Attribute_EndingAlpha = "EndingAlpha";
+		const FName Attribute_Gradient_Angle = "Gradient_Angle";
+	}
+};
+
+
+/** Actor data that can possibly be set to the light card actor */
+struct FDMXDisplayClusterLightCardActorData
+{
+	void Apply(ADisplayClusterLightCardActor* Actor)
+	{
+		if (!Actor || Actor->bIsProxy) 
+		{
+			return;
+		}
+
+		// Early out if DMX Input is disabled
+		if (!bDMXInput.IsSet() || !bDMXInput.GetValue())
+		{
+			return;
+		}
+
+		bool bNeedsUpdateLightCardTransform = false;
+		bool bNeedsUpdateLightCardMaterialInstance = false;
+		bool bNeedsUpdatePolygonTexture = false;
+
+		// Set properties common to all modes
+#if WITH_EDITOR
+		Actor->PreEditChange(nullptr);
+#endif
+
+		if (DistanceFromCenter.IsSet() && Actor->DistanceFromCenter != DistanceFromCenter)
+		{
+			Actor->DistanceFromCenter = DistanceFromCenter.GetValue();
+			bNeedsUpdateLightCardTransform = true;
+		}
+
+		if (Actor->bIsUVLightCard && 
+			Longitude.IsSet() && Latitude.IsSet() && 
+			(Actor->UVCoordinates.X != Longitude.GetValue() || Actor->UVCoordinates.Y != Latitude.GetValue()))
+		{
+			// Set UV Coordinates if in UV mode
+			Actor->UVCoordinates = FVector2D(Longitude.GetValue(), Latitude.GetValue());
+		}
+		else
+		{
+			// Set actual Longitude and Latidude if not in UV Mode
+			if (Longitude.IsSet() && Actor->Longitude != Longitude)
+			{
+				Actor->Longitude = Longitude.GetValue();
+				bNeedsUpdateLightCardTransform = true;
+			}
+
+			if (Latitude.IsSet() && Actor->Latitude != Latitude)
+			{
+				Actor->Latitude = Latitude.GetValue();
+				bNeedsUpdateLightCardTransform = true;
+			}
+		}
+
+
+		if (Spin.IsSet() && Actor->Spin != Spin)
+		{
+			Actor->Spin = Spin.GetValue();
+			bNeedsUpdateLightCardTransform = true;
+		}
+
+		if (Yaw.IsSet() && Actor->Yaw != Yaw)
+		{
+			Actor->Yaw = Yaw.GetValue();
+			bNeedsUpdateLightCardTransform = true;
+		}
+
+		if (Pitch.IsSet() && Actor->Pitch != Pitch)
+		{
+			Actor->Pitch = Pitch.GetValue();
+			bNeedsUpdateLightCardTransform = true;
+		}
+
+		if (ScaleX.IsSet() && Actor->Scale.X != ScaleX)
+		{
+			Actor->Scale.X = ScaleX.GetValue();
+			bNeedsUpdateLightCardTransform = true;
+		}
+
+		if (ScaleY.IsSet() && Actor->Scale.Y != ScaleY)
+		{
+			Actor->Scale.Y = ScaleY.GetValue();
+			bNeedsUpdateLightCardTransform = true;
+		}
+
+		if (Mask.IsSet() && Actor->Mask != Mask)
+		{
+			Actor->Mask = Mask.GetValue();
+			bNeedsUpdateLightCardMaterialInstance = true;
+		}
+
+		if (ColorR.IsSet() && Actor->Color.R != ColorR)
+		{
+			Actor->Color.R = ColorR.GetValue();
+			bNeedsUpdateLightCardMaterialInstance = true;
+		}
+
+		if (ColorG.IsSet() && Actor->Color.G != ColorG)
+		{
+			Actor->Color.G = ColorG.GetValue();
+			bNeedsUpdateLightCardMaterialInstance = true;
+		}
+
+		if (ColorB.IsSet() && Actor->Color.B != ColorB)
+		{
+			Actor->Color.B = ColorB.GetValue();
+			bNeedsUpdateLightCardMaterialInstance = true;
+		}
+
+		if (ColorA.IsSet() && Actor->Color.A != ColorA)
+		{
+			Actor->Color.A = ColorA.GetValue();
+			bNeedsUpdateLightCardMaterialInstance = true;
+		}
+
+		if (Temperature.IsSet() && Actor->Temperature != Temperature)
+		{
+			Actor->Temperature = Temperature.GetValue();
+			bNeedsUpdateLightCardMaterialInstance = true;
+		}
+
+		if (Tint.IsSet() && Actor->Tint != Tint)
+		{
+			Actor->Tint = Tint.GetValue();
+			bNeedsUpdateLightCardMaterialInstance = true;
+		}
+
+		if (Exposure.IsSet() && Actor->Exposure != Exposure)
+		{
+			Actor->Opacity = Exposure.GetValue();
+			bNeedsUpdateLightCardMaterialInstance = true;
+		}
+
+		if (Gain.IsSet() && Actor->Gain != Gain)
+		{
+			Actor->Gain = Gain.GetValue();
+			bNeedsUpdateLightCardMaterialInstance = true;
+		}
+
+		if (Opacity.IsSet() && Actor->Opacity != Opacity)
+		{
+			Actor->Opacity = Opacity.GetValue();
+			bNeedsUpdateLightCardMaterialInstance = true;
+		}
+
+		if (Feathering.IsSet() && Actor->Feathering != Feathering)
+		{
+			Actor->Feathering = Feathering.GetValue();
+			bNeedsUpdatePolygonTexture = true;
+		}
+
+		if (bAlphaGradientEnable.IsSet() && Actor->AlphaGradient.bEnableAlphaGradient != bAlphaGradientEnable)
+		{
+			Actor->AlphaGradient.bEnableAlphaGradient = bAlphaGradientEnable.GetValue();
+			bNeedsUpdatePolygonTexture = true;
+		}
+		
+		// Apply Alpha gradient only if Alpha Gradient Enable
+		if (IsAlphaGradientEnabled())
+		{
+			if (StartingAlpha.IsSet() && Actor->AlphaGradient.StartingAlpha != StartingAlpha)
+			{
+				Actor->AlphaGradient.StartingAlpha = StartingAlpha.GetValue();
+				bNeedsUpdatePolygonTexture = true;
+			}
+
+			if (EndingAlpha.IsSet() && Actor->AlphaGradient.EndingAlpha != EndingAlpha)
+			{
+				Actor->AlphaGradient.EndingAlpha = EndingAlpha.GetValue();
+				bNeedsUpdatePolygonTexture = true;
+			}
+
+			if (GradientAngle.IsSet() && Actor->AlphaGradient.Angle != GradientAngle)
+			{
+				Actor->AlphaGradient.Angle = GradientAngle.GetValue();
+				bNeedsUpdatePolygonTexture = true;
+			}
+		}
+
+		// Apply changes
+		if (bNeedsUpdateLightCardTransform)
+		{
+			Actor->UpdateLightCardTransform();
+		}
+
+		if (bNeedsUpdateLightCardMaterialInstance)
+		{
+			Actor->UpdateLightCardMaterialInstance();
+		}
+
+		if (bNeedsUpdatePolygonTexture)
+		{
+			Actor->UpdatePolygonTexture();
+		}
+
+#if WITH_EDITOR
+		Actor->PostEditChange();
+#endif
+	}
+
+	void SetDMXInput(uint8 ShouldInputDMXValue)
+	{
+		bDMXInput = ShouldInputDMXValue > DMX_MAX_VALUE / 2;
+	}
+
+	void SetDistanceFromCenter(double NormalizedValue, const FDMXDisplayClusterLightCardActorDataValueRanges& ValueRanges)
+	{
+		DistanceFromCenter = FMath::Lerp(ValueRanges.MinDistanceFromCenter, ValueRanges.MaxDistanceFromCenter, NormalizedValue);
+	}
+
+	void SetLongitude(double NormalizedValue, const FDMXDisplayClusterLightCardActorDataValueRanges& ValueRanges, bool bIsUVLightCard)
+	{
+		if (bIsUVLightCard)
+		{
+			Longitude = FMath::Lerp(ValueRanges.MinLongitudeU, ValueRanges.MaxLongitudeU, NormalizedValue);
+		}
+		else
+		{
+			Longitude = FMath::Lerp(ValueRanges.MinLongitude, ValueRanges.MaxLongitude, NormalizedValue);
+		}
+	}
+
+	void SetLatitude(double NormalizedValue, const FDMXDisplayClusterLightCardActorDataValueRanges& ValueRanges, bool bIsUVLightCard)
+	{
+		if (bIsUVLightCard)
+		{
+			Latitude = FMath::Lerp(ValueRanges.MinLatitudeV, ValueRanges.MaxLatitudeV, NormalizedValue);
+		}
+		else
+		{
+			Latitude = FMath::Lerp(ValueRanges.MinLatitude, ValueRanges.MaxLatitude, NormalizedValue);
+		}
+	}
+
+	void SetSpin(double NormalizedValue, const FDMXDisplayClusterLightCardActorDataValueRanges& ValueRanges)
+	{
+		Spin = FMath::Lerp(ValueRanges.MinSpin, ValueRanges.MaxSpin, NormalizedValue);
+	}
+
+	void SetPitch(double NormalizedValue, const FDMXDisplayClusterLightCardActorDataValueRanges& ValueRanges)
+	{
+		Pitch = FMath::Lerp(ValueRanges.MinPitch, ValueRanges.MaxPitch, NormalizedValue);
+	}
+
+	void SetYaw(double NormalizedValue, const FDMXDisplayClusterLightCardActorDataValueRanges& ValueRanges)
+	{
+		Yaw = FMath::Lerp(ValueRanges.MinYaw, ValueRanges.MaxYaw, NormalizedValue);
+	}
+
+	void SetScaleX(double NormalizedValue, const FDMXDisplayClusterLightCardActorDataValueRanges& ValueRanges)
+	{
+		ScaleX = FMath::Lerp(ValueRanges.MinScale.X, ValueRanges.MaxScale.X, NormalizedValue);
+	}
+	
+	void SetScaleY(double NormalizedValue, const FDMXDisplayClusterLightCardActorDataValueRanges& ValueRanges)
+	{
+		ScaleY = FMath::Lerp(ValueRanges.MinScale.Y, ValueRanges.MaxScale.Y, NormalizedValue);
+	}
+
+	void SetMask(uint8 AbsoluteValue)
+	{
+		Mask = static_cast<EDisplayClusterLightCardMask>(FMath::Clamp(AbsoluteValue, 0, static_cast<uint8>(EDisplayClusterLightCardMask::MAX) - 1));
+	}
+
+	void SetColorR(double NormalizedValue)
+	{
+		ColorR = NormalizedValue;;
+	}
+
+	void SetColorG(double NormalizedValue)
+	{
+		ColorG = NormalizedValue;;
+	}
+
+	void SetColorB(double NormalizedValue)
+	{
+		ColorB = NormalizedValue;;
+	}
+
+	void SetColorA(double NormalizedValue)
+	{
+		ColorA = NormalizedValue;;
+	}
+
+	void SetTemperature(double NormalizedValue, const FDMXDisplayClusterLightCardActorDataValueRanges& ValueRanges)
+	{
+		Temperature = FMath::Lerp(ValueRanges.MinTemperature, ValueRanges.MaxTemperature, NormalizedValue);
+	}
+
+	void SetTint(double NormalizedValue, const FDMXDisplayClusterLightCardActorDataValueRanges& ValueRanges)
+	{
+		Tint = FMath::Lerp(ValueRanges.MinTint, ValueRanges.MaxTint, NormalizedValue);
+	}
+
+	void SetExposure(double NormalizedValue, const FDMXDisplayClusterLightCardActorDataValueRanges& ValueRanges)
+	{
+		Exposure = FMath::Lerp(ValueRanges.MinExposure, ValueRanges.MaxExposure, NormalizedValue);
+	}
+
+	void SetGain(double NormalizedValue, const FDMXDisplayClusterLightCardActorDataValueRanges& ValueRanges)
+	{
+		Gain = FMath::Lerp(ValueRanges.MinGain, ValueRanges.MaxGain, NormalizedValue);
+	}
+
+	void SetOpacity(double NormalizedValue)
+	{
+		Opacity = NormalizedValue;
+	}
+
+	void SetFeathering(double NormalizedValue, const FDMXDisplayClusterLightCardActorDataValueRanges& ValueRanges)
+	{
+		Feathering = FMath::Lerp(ValueRanges.MinFeathering, ValueRanges.MaxFeathering, NormalizedValue);
+	}
+
+	void SetAlphaGradientEnable(uint8 AlphaGradientEnableValue)
+	{
+		bAlphaGradientEnable = AlphaGradientEnableValue > DMX_MAX_VALUE / 2;
+	}
+
+	void SetStartingAlpha(double NormalizedValue)
+	{
+		StartingAlpha = NormalizedValue;
+	}
+
+	void SetEndingAlpha(double NormalizedValue)
+	{
+		EndingAlpha = NormalizedValue;
+	}
+
+	void SetGradientAngle(double NormalizedValue, const FDMXDisplayClusterLightCardActorDataValueRanges& ValueRanges)
+	{
+		GradientAngle = FMath::Lerp(ValueRanges.MinGradientAngle, ValueRanges.MaxGradientAngle, NormalizedValue);
+	}
+
+	bool IsDMXInputEnabled() const { return bDMXInput.IsSet() && bDMXInput.GetValue(); }
+	bool IsAlphaGradientEnabled() const { return bAlphaGradientEnable.IsSet() && bAlphaGradientEnable.GetValue(); }
+
+private:
+	TOptional<bool> bDMXInput;
+	TOptional<double> DistanceFromCenter;
+	TOptional<double> Longitude;
+	TOptional<double> Latitude;
+	TOptional<double> Spin;
+	TOptional<double> Pitch;
+	TOptional<double> Yaw;
+	TOptional<double> ScaleX;
+	TOptional<double> ScaleY;
+	TOptional<EDisplayClusterLightCardMask> Mask;
+	TOptional<float> ColorR;
+	TOptional<float> ColorG;
+	TOptional<float> ColorB;
+	TOptional<float> ColorA;
+	TOptional<float> Temperature;
+	TOptional<float> Tint;
+	TOptional<float> Exposure;
+	TOptional<float> Gain;
+	TOptional<float> Opacity;
+	TOptional<float> Feathering;
+	TOptional<bool> bAlphaGradientEnable;
+	TOptional<float> StartingAlpha;
+	TOptional<float> EndingAlpha;
+	TOptional<float> GradientAngle;
+};
+
+void UDMXDisplayClusterLightCardComponent::OnRegister()
+{
+	Super::OnRegister();
+
+	if (!OnFixturePatchReceived.IsBound())
+	{
+		OnFixturePatchReceived.AddDynamic(this, &UDMXDisplayClusterLightCardComponent::OnLightCardReceivedDMXFromPatch);
+	}
+}
+
+void UDMXDisplayClusterLightCardComponent::PostLoad()
+{
+	Super::PostLoad();
+
+	if (!OnFixturePatchReceived.IsBound())
+	{
+		OnFixturePatchReceived.AddDynamic(this, &UDMXDisplayClusterLightCardComponent::OnLightCardReceivedDMXFromPatch);
+	}
+}
+
+#if WITH_EDITOR
+void UDMXDisplayClusterLightCardComponent::PostEditChangeProperty(FPropertyChangedEvent& PropertyChangedEvent)
+{
+	Super::PostEditChangeProperty(PropertyChangedEvent);
+
+	// Enable tick if no fixture patch is set and vice versa
+	if (PropertyChangedEvent.GetPropertyName() == GET_MEMBER_NAME_CHECKED(UDMXComponent, FixturePatchRef))
+	{
+		if (FixturePatchRef.GetFixturePatch())
+		{
+			if (!OnFixturePatchReceived.IsBound())
+			{
+				OnFixturePatchReceived.AddDynamic(this, &UDMXDisplayClusterLightCardComponent::OnLightCardReceivedDMXFromPatch);
+			}
+		}
+	}
+}
+#endif // WITH_EDITOR
+
+void UDMXDisplayClusterLightCardComponent::OnLightCardReceivedDMXFromPatch(UDMXEntityFixturePatch* FixturePatch, const FDMXNormalizedAttributeValueMap& ValuePerAttribute)
+{
+	using namespace UE::DMX::DMXDisplayClusterLightCardComponent::Private::AttributeNames;
+
+	ADisplayClusterLightCardActor* Actor = Cast<ADisplayClusterLightCardActor>(GetOwner());
+	if (!Actor)
+	{
+		return;
+	}
+
+	bool bAccessDMXInputValueSuccess;
+	const uint8 DMXInputValue = FixturePatch->GetAttributeValue(Attribute_DMXInput, bAccessDMXInputValueSuccess);
+	if (!bAccessDMXInputValueSuccess)
+	{
+		// Don't proceed if DMXInput is not set
+		return;
+	}
+
+	FDMXDisplayClusterLightCardActorData ActorData;
+	ActorData.SetDMXInput(DMXInputValue);
+
+	// Early out if DMX Input is disabled
+	if (!ActorData.IsDMXInputEnabled())
+	{
+		return;
+	}
+
+	const float* const DistanceFromCenterPtr = ValuePerAttribute.Map.Find(Attribute_DistanceFromCenter);
+	if (DistanceFromCenterPtr)
+	{
+		ActorData.SetDistanceFromCenter(*DistanceFromCenterPtr, ValueRanges);
+	}
+
+	const float* const LongitudePtr = ValuePerAttribute.Map.Find(Attribute_Pan);
+	if (LongitudePtr)
+	{
+		ActorData.SetLongitude(*LongitudePtr, ValueRanges, Actor->bIsUVLightCard);
+	}
+
+	const float* const LatitudePtr = ValuePerAttribute.Map.Find(Attribute_Tilt);
+	if (LatitudePtr)
+	{
+		ActorData.SetLatitude(*LatitudePtr, ValueRanges, Actor->bIsUVLightCard);
+	}
+
+	const float* const SpinPtr = ValuePerAttribute.Map.Find(Attribute_RotX);
+	if (SpinPtr)
+	{
+		ActorData.SetSpin(*SpinPtr, ValueRanges);
+	}
+
+	const float* const PitchPtr = ValuePerAttribute.Map.Find(Attribute_RotY);
+	if (PitchPtr)
+	{
+		ActorData.SetPitch(*PitchPtr, ValueRanges);
+	}
+
+	const float* const YawPtr = ValuePerAttribute.Map.Find(Attribute_RotZ);
+	if (YawPtr)
+	{
+		ActorData.SetYaw(*YawPtr, ValueRanges);
+	}
+
+	const float* const ScaleXPtr = ValuePerAttribute.Map.Find(Attribute_ScaleX);
+	if (ScaleXPtr)
+	{
+		ActorData.SetScaleX(*ScaleXPtr, ValueRanges);
+	}
+
+	const float* const ScaleYPtr = ValuePerAttribute.Map.Find(Attribute_ScaleY);
+	if (ScaleYPtr)
+	{
+		ActorData.SetScaleY(*ScaleYPtr, ValueRanges);
+	}
+
+	bool bAccessMaskValueSuccess;
+	const uint8 Mask = FixturePatch->GetAttributeValue(Attribute_Mask, bAccessMaskValueSuccess);
+	if (bAccessMaskValueSuccess)
+	{
+		ActorData.SetMask(Mask);
+	}
+
+	const float* const RedPtr = ValuePerAttribute.Map.Find(Attribute_Red);
+	if (RedPtr)
+	{
+		ActorData.SetColorR(*RedPtr);
+	}
+
+	const float* const GreenPtr = ValuePerAttribute.Map.Find(Attribute_Green);
+	if (GreenPtr)
+	{
+		ActorData.SetColorG(*GreenPtr);
+	}
+
+	const float* const BluePtr = ValuePerAttribute.Map.Find(Attribute_Blue);
+	if (BluePtr)
+	{
+		ActorData.SetColorB(*BluePtr);
+	}
+
+	const float* const AlphaPtr = ValuePerAttribute.Map.Find(Attribute_ColorAdd_Alpha);
+	if (AlphaPtr)
+	{
+		ActorData.SetColorA(*AlphaPtr);
+	}
+
+	const float* const CTCPtr = ValuePerAttribute.Map.Find(Attribute_CTC);
+	if (CTCPtr)
+	{
+		ActorData.SetTemperature(*CTCPtr, ValueRanges);
+	}
+
+	const float* const ExposurePtr = ValuePerAttribute.Map.Find(Attribute_Exposure);
+	if (ExposurePtr)
+	{
+		ActorData.SetExposure(*ExposurePtr, ValueRanges);
+	}
+
+	const float* const GainPtr = ValuePerAttribute.Map.Find(Attribute_Gain);
+	if (GainPtr)
+	{
+		ActorData.SetGain(*GainPtr, ValueRanges);
+	}
+
+	const float* const OpacityPtr = ValuePerAttribute.Map.Find(Attribute_Opacity);
+	if (OpacityPtr)
+	{
+		ActorData.SetOpacity(*OpacityPtr);
+	}
+
+	const float* const FeatheringPtr = ValuePerAttribute.Map.Find(Attribute_Feathering);
+	if (FeatheringPtr)
+	{
+		ActorData.SetFeathering(*FeatheringPtr, ValueRanges);
+	}
+
+	bool bAccessAlphaGradientEnableValueSuccess;
+	const int32 AlphaGradientEnableValue = FixturePatch->GetAttributeValue(Attribute_Alpha_Gradient_Enable, bAccessAlphaGradientEnableValueSuccess);
+	if (bAccessAlphaGradientEnableValueSuccess)
+	{
+		ActorData.SetAlphaGradientEnable(AlphaGradientEnableValue);
+	}
+
+	if (ActorData.IsAlphaGradientEnabled())
+	{
+		const float* const StartingAlphaPtr = ValuePerAttribute.Map.Find(Attribute_StartingAlpha);
+		if (StartingAlphaPtr)
+		{
+			ActorData.SetStartingAlpha(*StartingAlphaPtr);
+		}
+
+		const float* const EndingAlphaPtr = ValuePerAttribute.Map.Find(Attribute_EndingAlpha);
+		if (EndingAlphaPtr)
+		{
+			ActorData.SetEndingAlpha(*EndingAlphaPtr);
+		}
+
+		const float* const GradientAnglePtr = ValuePerAttribute.Map.Find(Attribute_Gradient_Angle);
+		if (GradientAnglePtr)
+		{
+			ActorData.SetGradientAngle(*GradientAnglePtr, ValueRanges);
+		}
+	}
+
+	ActorData.Apply(Actor);
+}
diff --git a/Engine/Plugins/VirtualProduction/DMX/DMXDisplayCluster/Source/DMXDisplayClusterLightCard/Private/DMXDisplayClusterLightCardComponent.h b/Engine/Plugins/VirtualProduction/DMX/DMXDisplayCluster/Source/DMXDisplayClusterLightCard/Private/DMXDisplayClusterLightCardComponent.h
new file mode 100644
--- /dev/null
+++ b/Engine/Plugins/VirtualProduction/DMX/DMXDisplayCluster/Source/DMXDisplayClusterLightCard/Private/DMXDisplayClusterLightCardComponent.h
@@ -0,0 +1,141 @@
+// Copyright Epic Games, Inc. All Rights Reserved.
+
+#pragma once
+
+#include "DisplayClusterLightCardActor.h"
+#include "IO/DMXInputPortReference.h"
+#include "Game/DMXComponent.h"
+
+#include "Components/ActorComponent.h"
+
+#include "DMXDisplayClusterLightCardComponent.generated.h"
+
+class UDMXEntityFixturePatch;
+
+
+
+USTRUCT()
+struct FDMXDisplayClusterLightCardActorDataValueRanges
+{
+	GENERATED_BODY()
+
+	UPROPERTY(EditAnywhere, Category = "DMX")
+	double MinDistanceFromCenter = 0.0;
+
+	UPROPERTY(EditAnywhere, Category = "DMX")
+	double MaxDistanceFromCenter = 1000.0;
+
+	UPROPERTY(EditAnywhere, Category = "DMX")
+	double MinLongitude = 0.0;
+
+	UPROPERTY(EditAnywhere, Category = "DMX")
+	double MaxLongitude = 360.0;
+
+	UPROPERTY(EditAnywhere, Category = "DMX")
+	double MinLatitude = -90.0;
+
+	/** Min longitude when the light card uses UV Mode */
+	UPROPERTY(EditAnywhere, Category = "DMX")
+	float MinLongitudeU = -2.f;
+	
+	/** Max longitude when the light card uses UV Mode */
+	UPROPERTY(EditAnywhere, Category = "DMX")
+	float MaxLongitudeU = 2.f;
+		
+	/** Min latitude when the light card uses UV Mode */
+	UPROPERTY(EditAnywhere, Category = "DMX")
+	float MinLatitudeV = -2.f;
+	
+	/** Max latitude when the light card uses UV Mode */
+	UPROPERTY(EditAnywhere, Category = "DMX")
+	float MaxLatitudeV = 2.f;
+
+	UPROPERTY(EditAnywhere, Category = "DMX")
+	double MaxLatitude = 90.0;
+
+	UPROPERTY(EditAnywhere, Category = "DMX")
+	double MinSpin = 0.0;
+
+	UPROPERTY(EditAnywhere, Category = "DMX")
+	double MaxSpin = 360.0;
+
+	UPROPERTY(EditAnywhere, Category = "DMX")
+	double MinPitch = 0.0;
+
+	UPROPERTY(EditAnywhere, Category = "DMX")
+	double MaxPitch = 360.0;
+
+	UPROPERTY(EditAnywhere, Category = "DMX")
+	double MinYaw = 0.0;
+
+	UPROPERTY(EditAnywhere, Category = "DMX")
+	double MaxYaw = 360.0;
+
+	UPROPERTY(EditAnywhere, Category = "DMX")
+	FVector2D MinScale = FVector2D(0.001f, 0.001f);
+
+	UPROPERTY(EditAnywhere, Category = "DMX")
+	FVector2D MaxScale = FVector2D(1.f, 1.f);
+
+	UPROPERTY(EditAnywhere, Category = "DMX")
+	float MinTemperature = 0.f;
+
+	UPROPERTY(EditAnywhere, Category = "DMX")
+	float MaxTemperature = 10000.f;
+
+	UPROPERTY(EditAnywhere, Category = "DMX")
+	float MinTint = -1.f;
+
+	UPROPERTY(EditAnywhere, Category = "DMX")
+	float MaxTint = 1.f;
+
+	UPROPERTY(EditAnywhere, Category = "DMX")
+	float MinExposure = -100.f;
+
+	UPROPERTY(EditAnywhere, Category = "DMX")
+	float MaxExposure = 100.f;
+
+	UPROPERTY(EditAnywhere, Category = "DMX")
+	float MinGain = 0.f;
+	
+	UPROPERTY(EditAnywhere, Category = "DMX")
+	float MaxGain = 1000.f;
+	
+	UPROPERTY(EditAnywhere, Category = "DMX")
+	float MinFeathering = 0.f;
+	
+	UPROPERTY(EditAnywhere, Category = "DMX")
+	float MaxFeathering = 1000.f;
+
+	UPROPERTY(EditAnywhere, Category = "DMX")
+	float MinGradientAngle = 0.f;
+	
+	UPROPERTY(EditAnywhere, Category = "DMX")
+	float MaxGradientAngle = 360.f;
+};
+
+UCLASS(NotBlueprintable, HideCategories = (AssetUserData, Tags, Replication))
+class UDMXDisplayClusterLightCardComponent
+	: public UDMXComponent
+{
+	GENERATED_BODY()
+
+protected:
+	//~ Begin UObject interface
+	virtual void OnRegister() override;
+	virtual void PostLoad() override;
+#if WITH_EDITOR
+	virtual void PostEditChangeProperty(FPropertyChangedEvent& PropertyChangedEvent) override;
+#endif
+	//~ End UObject interface
+
+public:
+	/** Value ranges for Min and Max DMX Values */
+	UPROPERTY(EditAnywhere, AdvancedDisplay, Category = "DMX", Meta = (ShowOnlyInnerProperties))
+	FDMXDisplayClusterLightCardActorDataValueRanges ValueRanges;
+
+private:
+	/** Called when the fixture patch received DMX */
+	UFUNCTION()
+	void OnLightCardReceivedDMXFromPatch(UDMXEntityFixturePatch* FixturePatch, const FDMXNormalizedAttributeValueMap& ValuePerAttribute);
+};
diff --git a/Engine/Plugins/VirtualProduction/DMX/DMXDisplayCluster/Source/DMXDisplayClusterLightCard/Private/DMXDisplayClusterLightCardModule.cpp b/Engine/Plugins/VirtualProduction/DMX/DMXDisplayCluster/Source/DMXDisplayClusterLightCard/Private/DMXDisplayClusterLightCardModule.cpp
new file mode 100644
--- /dev/null
+++ b/Engine/Plugins/VirtualProduction/DMX/DMXDisplayCluster/Source/DMXDisplayClusterLightCard/Private/DMXDisplayClusterLightCardModule.cpp
@@ -0,0 +1,45 @@
+// Copyright Epic Games, Inc. All Rights Reserved.
+
+#include "DMXDisplayClusterLightCardModule.h"
+
+#include "DMXDisplayClusterLightCardComponent.h"
+
+#include "Features/IModularFeatures.h"
+#include "Modules/ModuleManager.h"
+
+
+void FDMXDisplayClusterLightCardModule::StartupModule()
+{
+	IModularFeatures::Get().RegisterModularFeature(IDisplayClusterLightCardActorExtender::ModularFeatureName, this);
+}
+
+void FDMXDisplayClusterLightCardModule::ShutdownModule()
+{
+	IModularFeatures::Get().UnregisterModularFeature(IDisplayClusterLightCardActorExtender::ModularFeatureName, this);
+}
+
+FName FDMXDisplayClusterLightCardModule::GetExtenderName() const
+{
+	return "DMXDisplayClusterLightCardExtender";
+}
+
+TSubclassOf<UActorComponent> FDMXDisplayClusterLightCardModule::GetAdditionalSubobjectClass()
+{
+	return UDMXDisplayClusterLightCardComponent::StaticClass();
+}
+
+#if WITH_EDITOR
+FName FDMXDisplayClusterLightCardModule::GetCategory() const
+{
+	return "DMX";
+}
+#endif
+
+#if WITH_EDITOR
+bool FDMXDisplayClusterLightCardModule::ShouldShowSubcategories() const
+{
+	return false;
+}
+#endif
+
+IMPLEMENT_MODULE(FDMXDisplayClusterLightCardModule, DMXDisplayClusterLightCard);
diff --git a/Engine/Plugins/VirtualProduction/DMX/DMXDisplayCluster/Source/DMXDisplayClusterLightCard/Private/DMXDisplayClusterLightCardModule.h b/Engine/Plugins/VirtualProduction/DMX/DMXDisplayCluster/Source/DMXDisplayClusterLightCard/Private/DMXDisplayClusterLightCardModule.h
new file mode 100644
--- /dev/null
+++ b/Engine/Plugins/VirtualProduction/DMX/DMXDisplayCluster/Source/DMXDisplayClusterLightCard/Private/DMXDisplayClusterLightCardModule.h
@@ -0,0 +1,29 @@
+// Copyright Epic Games, Inc. All Rights Reserved.
+
+#pragma once
+
+#include "IDisplayClusterLightCardActorExtender.h"
+
+#include "CoreMinimal.h"
+#include "Modules/ModuleInterface.h"
+
+
+class FDMXDisplayClusterLightCardModule 
+	: public IModuleInterface
+	, public IDisplayClusterLightCardActorExtender
+{
+public:
+	//~ Begin IModuleInterface inteface
+	virtual void StartupModule() override;
+	virtual void ShutdownModule() override;
+	//~ End IModuleInterface inteface
+
+	//~ Begin IDisplayClusterLightCardActorExtender interface
+	virtual FName GetExtenderName() const override;
+	virtual TSubclassOf<UActorComponent> GetAdditionalSubobjectClass()override;
+#if WITH_EDITOR
+	virtual FName GetCategory() const override;
+	virtual bool ShouldShowSubcategories() const override;
+#endif
+	//~ End IDisplayClusterLightCardActorExtender interface
+};
diff --git a/Engine/Plugins/VirtualProduction/DMX/DMXEngine/Source/DMXRuntime/Public/Game/DMXComponent.h b/Engine/Plugins/VirtualProduction/DMX/DMXEngine/Source/DMXRuntime/Public/Game/DMXComponent.h
--- a/Engine/Plugins/VirtualProduction/DMX/DMXEngine/Source/DMXRuntime/Public/Game/DMXComponent.h
+++ b/Engine/Plugins/VirtualProduction/DMX/DMXEngine/Source/DMXRuntime/Public/Game/DMXComponent.h
@@ -55,7 +55,7 @@ public:
 	UFUNCTION(BlueprintCallable, Category = "DMX")
 	void SetReceiveDMXFromPatch(bool bReceive);
 
-	UPROPERTY(BlueprintReadOnly, EditAnywhere, Category = "DMX")
+	UPROPERTY(BlueprintReadOnly, EditAnywhere, Category = "DMX", Meta = (ShowOnlyInnerProperties))
 	FDMXEntityFixturePatchRef FixturePatchRef;
 
 protected:
diff --git a/Engine/Plugins/VirtualProduction/DMX/DMXProtocol/Source/DMXProtocolEditor/Private/DetailsCustomizations/DMXPortReferenceCustomizationBase.cpp b/Engine/Plugins/VirtualProduction/DMX/DMXProtocol/Source/DMXProtocolEditor/Private/DetailsCustomizations/DMXPortReferenceCustomizationBase.cpp
--- a/Engine/Plugins/VirtualProduction/DMX/DMXProtocol/Source/DMXProtocolEditor/Private/DetailsCustomizations/DMXPortReferenceCustomizationBase.cpp
+++ b/Engine/Plugins/VirtualProduction/DMX/DMXProtocol/Source/DMXProtocolEditor/Private/DetailsCustomizations/DMXPortReferenceCustomizationBase.cpp
@@ -41,10 +41,6 @@ void FDMXPortReferenceCustomizationBase::CustomizeChildren(TSharedRef<IPropertyH
 	// Hide the 'reset to default' option
 	StructPropertyHandle->MarkResetToDefaultCustomized();
 
-	// Create the info content border, so the port can set its info during its construction
-	SAssignNew(InfoContentBorder, SBorder)
-		.BorderImage(FAppStyle::GetBrush("NoBorder"));
-
 	// Add the port selector row
 	EDMXPortSelectorMode PortSelectorMode = IsInputPort() ? EDMXPortSelectorMode::SelectFromAvailableInputs : EDMXPortSelectorMode::SelectFromAvailableOutputs;
 
@@ -63,13 +59,6 @@ void FDMXPortReferenceCustomizationBase::CustomizeChildren(TSharedRef<IPropertyH
 			.OnPortSelected(this, &FDMXPortReferenceCustomizationBase::OnPortSelected)
 		];
 
-	// Add the info row
-	ChildBuilder.AddCustomRow(LOCTEXT("PortInfoSearchString", "Info"))
-		.WholeRowContent()
-		[
-			InfoContentBorder.ToSharedRef()
-		];
-
 	// Four possible states here that all need be handled:
 	// a) This is an existing port reference, with a corresponding port config
 	// b) This is an existing port reference, but the corresponding port config got deleted in settings
diff --git a/Engine/Plugins/VirtualProduction/DMX/DMXProtocol/Source/DMXProtocolEditor/Private/DetailsCustomizations/DMXPortReferenceCustomizationBase.h b/Engine/Plugins/VirtualProduction/DMX/DMXProtocol/Source/DMXProtocolEditor/Private/DetailsCustomizations/DMXPortReferenceCustomizationBase.h
--- a/Engine/Plugins/VirtualProduction/DMX/DMXProtocol/Source/DMXProtocolEditor/Private/DetailsCustomizations/DMXPortReferenceCustomizationBase.h
+++ b/Engine/Plugins/VirtualProduction/DMX/DMXProtocol/Source/DMXProtocolEditor/Private/DetailsCustomizations/DMXPortReferenceCustomizationBase.h
@@ -56,9 +56,6 @@ private:
 	/** Port selector widget */
 	TSharedPtr<SDMXPortSelector> PortSelector;
 
-	/** Border to hold info content */
-	TSharedPtr<SBorder> InfoContentBorder;
-
 	/** Error text shown when the port couldn't be found, may contain the last error if there is no error - I11t is not a state. */
 	FText ErrorText;
 
