DMX - Improve value ranges of DMX controls for light cards

diff --git a/Engine/Plugins/VirtualProduction/DMX/DMXDisplayCluster/Source/DMXDisplayClusterLightCard/Private/DMXDisplayClusterLightCardComponent.cpp b/Engine/Plugins/VirtualProduction/DMX/DMXDisplayCluster/Source/DMXDisplayClusterLightCard/Private/DMXDisplayClusterLightCardComponent.cpp
--- a/Engine/Plugins/VirtualProduction/DMX/DMXDisplayCluster/Source/DMXDisplayClusterLightCard/Private/DMXDisplayClusterLightCardComponent.cpp
+++ b/Engine/Plugins/VirtualProduction/DMX/DMXDisplayCluster/Source/DMXDisplayClusterLightCard/Private/DMXDisplayClusterLightCardComponent.cpp
@@ -18,11 +18,11 @@ namespace UE::DMX::DMXDisplayClusterLightCardComponent::Private
 		const FName Attribute_DistanceFromCenter = "DistanceFromCenter";
 		const FName Attribute_Pan = "Pan";
 		const FName Attribute_Tilt = "Tilt";
-		const FName Attribute_RotX = "Spin";
-		const FName Attribute_RotY = "Pitch";
-		const FName Attribute_RotZ = "Yaw";
-		const FName Attribute_ScaleX = "ScaleX";
-		const FName Attribute_ScaleY = "ScaleY";
+		const FName Attribute_RotX = "Rot_X";
+		const FName Attribute_RotY = "Rot_Y";
+		const FName Attribute_RotZ = "Rot_Z";
+		const FName Attribute_ScaleX = "Scale_X";
+		const FName Attribute_ScaleY = "Scale_Y";
 		const FName Attribute_Mask = "Mask";
 		const FName Attribute_Red = "Red";
 		const FName Attribute_Green = "Green";
@@ -303,7 +303,15 @@ struct FDMXDisplayClusterLightCardActorData
 
 	void SetMask(uint8 AbsoluteValue)
 	{
-		Mask = static_cast<EDisplayClusterLightCardMask>(FMath::Clamp(AbsoluteValue, 0, static_cast<uint8>(EDisplayClusterLightCardMask::MAX) - 1));
+		// Skip through if the value is 0
+		if (AbsoluteValue == 0)
+		{
+			return;
+		}
+
+		// Scale the value to ranges of 5, so 1-5 is the first enum value, 6-10 the second, etc. 
+		const uint8 AbsoluteEnumValue = (AbsoluteValue + 4) / 5 - 1;
+		Mask = static_cast<EDisplayClusterLightCardMask>(FMath::Clamp(AbsoluteEnumValue, 0, static_cast<uint8>(EDisplayClusterLightCardMask::MAX) - 1));
 	}
 
 	void SetColorR(double NormalizedValue)
diff --git a/Engine/Plugins/VirtualProduction/DMX/DMXDisplayCluster/Source/DMXDisplayClusterLightCard/Private/DMXDisplayClusterLightCardComponent.h b/Engine/Plugins/VirtualProduction/DMX/DMXDisplayCluster/Source/DMXDisplayClusterLightCard/Private/DMXDisplayClusterLightCardComponent.h
--- a/Engine/Plugins/VirtualProduction/DMX/DMXDisplayCluster/Source/DMXDisplayClusterLightCard/Private/DMXDisplayClusterLightCardComponent.h
+++ b/Engine/Plugins/VirtualProduction/DMX/DMXDisplayCluster/Source/DMXDisplayClusterLightCard/Private/DMXDisplayClusterLightCardComponent.h
@@ -72,10 +72,10 @@ struct FDMXDisplayClusterLightCardActorDataValueRanges
 	double MaxYaw = 360.0;
 
 	UPROPERTY(EditAnywhere, Category = "DMX")
-	FVector2D MinScale = FVector2D(0.001f, 0.001f);
+	FVector2D MinScale = FVector2D(0.f, 0.f);
 
 	UPROPERTY(EditAnywhere, Category = "DMX")
-	FVector2D MaxScale = FVector2D(1.f, 1.f);
+	FVector2D MaxScale = FVector2D(5.f, 5.f);
 
 	UPROPERTY(EditAnywhere, Category = "DMX")
 	float MinTemperature = 0.f;
@@ -99,13 +99,13 @@ struct FDMXDisplayClusterLightCardActorDataValueRanges
 	float MinGain = 0.f;
 	
 	UPROPERTY(EditAnywhere, Category = "DMX")
-	float MaxGain = 1000.f;
+	float MaxGain = 5.f;
 	
 	UPROPERTY(EditAnywhere, Category = "DMX")
 	float MinFeathering = 0.f;
 	
 	UPROPERTY(EditAnywhere, Category = "DMX")
-	float MaxFeathering = 1000.f;
+	float MaxFeathering = 3.f;
 
 	UPROPERTY(EditAnywhere, Category = "DMX")
 	float MinGradientAngle = 0.f;
