DMX - Fix issues when setting exposure and tint for light cards

diff --git a/Engine/Plugins/VirtualProduction/DMX/DMXDisplayCluster/Source/DMXDisplayClusterLightCard/Private/DMXDisplayClusterLightCardComponent.cpp b/Engine/Plugins/VirtualProduction/DMX/DMXDisplayCluster/Source/DMXDisplayClusterLightCard/Private/DMXDisplayClusterLightCardComponent.cpp
--- a/Engine/Plugins/VirtualProduction/DMX/DMXDisplayCluster/Source/DMXDisplayClusterLightCard/Private/DMXDisplayClusterLightCardComponent.cpp
+++ b/Engine/Plugins/VirtualProduction/DMX/DMXDisplayCluster/Source/DMXDisplayClusterLightCard/Private/DMXDisplayClusterLightCardComponent.cpp
@@ -62,7 +62,6 @@ struct FDMXDisplayClusterLightCardActorData
 		bool bNeedsUpdateLightCardMaterialInstance = false;
 		bool bNeedsUpdatePolygonTexture = false;
 
-		// Set properties common to all modes
 #if WITH_EDITOR
 		Actor->PreEditChange(nullptr);
 #endif
@@ -171,7 +170,7 @@ struct FDMXDisplayClusterLightCardActorData
 
 		if (Exposure.IsSet() && Actor->Exposure != Exposure)
 		{
-			Actor->Opacity = Exposure.GetValue();
+			Actor->Exposure = Exposure.GetValue();
 			bNeedsUpdateLightCardMaterialInstance = true;
 		}
 
@@ -565,6 +564,12 @@ void UDMXDisplayClusterLightCardComponent::OnLightCardReceivedDMXFromPatch(UDMXE
 		ActorData.SetTemperature(*CTCPtr, ValueRanges);
 	}
 
+	const float* const TintPtr = ValuePerAttribute.Map.Find(Attribute_Tint);
+	if (TintPtr)
+	{
+		ActorData.SetTint(*TintPtr, ValueRanges);
+	}
+
 	const float* const ExposurePtr = ValuePerAttribute.Map.Find(Attribute_Exposure);
 	if (ExposurePtr)
 	{
