DMX - Omit any kind of transaction when setting light card properties

diff --git a/Engine/Plugins/VirtualProduction/DMX/DMXDisplayCluster/Source/DMXDisplayClusterLightCard/Private/DMXDisplayClusterLightCardComponent.cpp b/Engine/Plugins/VirtualProduction/DMX/DMXDisplayCluster/Source/DMXDisplayClusterLightCard/Private/DMXDisplayClusterLightCardComponent.cpp
--- a/Engine/Plugins/VirtualProduction/DMX/DMXDisplayCluster/Source/DMXDisplayClusterLightCard/Private/DMXDisplayClusterLightCardComponent.cpp
+++ b/Engine/Plugins/VirtualProduction/DMX/DMXDisplayCluster/Source/DMXDisplayClusterLightCard/Private/DMXDisplayClusterLightCardComponent.cpp
@@ -62,10 +62,6 @@ struct FDMXDisplayClusterLightCardActorData
 		bool bNeedsUpdateLightCardMaterialInstance = false;
 		bool bNeedsUpdatePolygonTexture = false;
 
-#if WITH_EDITOR
-		Actor->PreEditChange(nullptr);
-#endif
-
 		if (DistanceFromCenter.IsSet() && Actor->DistanceFromCenter != DistanceFromCenter)
 		{
 			Actor->DistanceFromCenter = DistanceFromCenter.GetValue();
@@ -235,10 +231,6 @@ struct FDMXDisplayClusterLightCardActorData
 		{
 			Actor->UpdatePolygonTexture();
 		}
-
-#if WITH_EDITOR
-		Actor->PostEditChange();
-#endif
 	}
 
 	void SetDMXInput(uint8 ShouldInputDMXValue)
