DMX - Adjust light card Exposure range to -5 to 5

diff --git a/Engine/Plugins/VirtualProduction/DMX/DMXDisplayCluster/Source/DMXDisplayClusterLightCard/Private/DMXDisplayClusterLightCardComponent.h b/Engine/Plugins/VirtualProduction/DMX/DMXDisplayCluster/Source/DMXDisplayClusterLightCard/Private/DMXDisplayClusterLightCardComponent.h
--- a/Engine/Plugins/VirtualProduction/DMX/DMXDisplayCluster/Source/DMXDisplayClusterLightCard/Private/DMXDisplayClusterLightCardComponent.h
+++ b/Engine/Plugins/VirtualProduction/DMX/DMXDisplayCluster/Source/DMXDisplayClusterLightCard/Private/DMXDisplayClusterLightCardComponent.h
@@ -90,10 +90,10 @@ struct FDMXDisplayClusterLightCardActorDataValueRanges
 	float MaxTint = 1.f;
 
 	UPROPERTY(EditAnywhere, Category = "DMX")
-	float MinExposure = -100.f;
+	float MinExposure = -5.f;
 
 	UPROPERTY(EditAnywhere, Category = "DMX")
-	float MaxExposure = 100.f;
+	float MaxExposure = 5.f;
 
 	UPROPERTY(EditAnywhere, Category = "DMX")
 	float MinGain = 0.f;
