nDisplay: Fixed issue where the DMXDisplayClusterLightCard module was not loading in before the light card actor archetype was initialized, which prevented the DMX light card component from being properly created on light card actor instances, and resulting in very poor framerate when a light card was selected in a detail panel.

diff --git a/Engine/Plugins/VirtualProduction/DMX/DMXDisplayCluster/DMXDisplayCluster.uplugin b/Engine/Plugins/VirtualProduction/DMX/DMXDisplayCluster/DMXDisplayCluster.uplugin
--- a/Engine/Plugins/VirtualProduction/DMX/DMXDisplayCluster/DMXDisplayCluster.uplugin
+++ b/Engine/Plugins/VirtualProduction/DMX/DMXDisplayCluster/DMXDisplayCluster.uplugin
@@ -25,7 +25,7 @@
 		{
 			"Name": "DMXDisplayClusterLightCard",
 			"Type": "Runtime",
-			"LoadingPhase": "Default",
+			"LoadingPhase": "PostConfigInit",
 			"PlatformAllowList": [
 				"Win64",
 				"Linux"
