nDisplay: Added invoke of OnObjectPropertyChanged to the DMX Light Card component when applying changes to the light card properties from a patch, allowing the ICVFX panel's light card proxies to be correctly updated to match the light card's new property values

diff --git a/Engine/Plugins/VirtualProduction/DMX/DMXDisplayCluster/Source/DMXDisplayClusterLightCard/Private/DMXDisplayClusterLightCardComponent.cpp b/Engine/Plugins/VirtualProduction/DMX/DMXDisplayCluster/Source/DMXDisplayClusterLightCard/Private/DMXDisplayClusterLightCardComponent.cpp
--- a/Engine/Plugins/VirtualProduction/DMX/DMXDisplayCluster/Source/DMXDisplayClusterLightCard/Private/DMXDisplayClusterLightCardComponent.cpp
+++ b/Engine/Plugins/VirtualProduction/DMX/DMXDisplayCluster/Source/DMXDisplayClusterLightCard/Private/DMXDisplayClusterLightCardComponent.cpp
@@ -237,6 +237,16 @@ struct FDMXDisplayClusterLightCardActorData
 		{
 			Actor->UpdatePolygonTexture();
 		}
+
+#if WITH_EDITOR
+		if (bNeedsUpdateLightCardTransform || bNeedsUpdateLightCardMaterialInstance || bNeedsUpdatePolygonTexture)
+		{
+			// Broadcast a property changed event for the light card so that anything listening for light card property changes, such as the ICVFX panel,
+			// can react appropriately. Since multiple properties may have changed, use an empty property event.
+			FPropertyChangedEvent EmptyPropertyChangedEvent(nullptr);
+			FCoreUObjectDelegates::OnObjectPropertyChanged.Broadcast(Actor, EmptyPropertyChangedEvent);
+		}
+#endif
 	}
 
 	void SetDMXInput(uint8 ShouldInputDMXValue)
