DMX: Let DMXDisplayClusterLightCardComponent follow the logic of latest DisplayClusterLightCardActor to update lightcards

diff --git a/Engine/Plugins/VirtualProduction/DMX/DMXDisplayCluster/Source/DMXDisplayClusterLightCard/Private/DMXDisplayClusterLightCardComponent.cpp b/Engine/Plugins/VirtualProduction/DMX/DMXDisplayCluster/Source/DMXDisplayClusterLightCard/Private/DMXDisplayClusterLightCardComponent.cpp
--- a/Engine/Plugins/VirtualProduction/DMX/DMXDisplayCluster/Source/DMXDisplayClusterLightCard/Private/DMXDisplayClusterLightCardComponent.cpp
+++ b/Engine/Plugins/VirtualProduction/DMX/DMXDisplayCluster/Source/DMXDisplayClusterLightCard/Private/DMXDisplayClusterLightCardComponent.cpp
@@ -47,7 +47,7 @@ struct FDMXDisplayClusterLightCardActorData
 {
 	void Apply(ADisplayClusterLightCardActor* Actor)
 	{
-		if (!Actor || Actor->IsProxy()) 
+		if (!Actor || Actor->IsProxy())
 		{
 			return;
 		}
@@ -68,12 +68,13 @@ struct FDMXDisplayClusterLightCardActorData
 			bNeedsUpdateLightCardTransform = true;
 		}
 
-		if (Actor->bIsUVLightCard && 
-			Longitude.IsSet() && Latitude.IsSet() && 
+		if (Actor->bIsUVLightCard &&
+			Longitude.IsSet() && Latitude.IsSet() &&
 			(Actor->UVCoordinates.X != Longitude.GetValue() || Actor->UVCoordinates.Y != Latitude.GetValue()))
 		{
 			// Set UV Coordinates if in UV mode
 			Actor->UVCoordinates = FVector2D(Longitude.GetValue(), Latitude.GetValue());
+			bNeedsUpdateLightCardTransform = true;
 		}
 		else
 		{
@@ -91,7 +92,6 @@ struct FDMXDisplayClusterLightCardActorData
 			}
 		}
 
-
 		if (Spin.IsSet() && Actor->Spin != Spin)
 		{
 			Actor->Spin = Spin.GetValue();
@@ -186,33 +186,34 @@ struct FDMXDisplayClusterLightCardActorData
 		{
 			Actor->Feathering = Feathering.GetValue();
 			bNeedsUpdatePolygonTexture = true;
+			bNeedsUpdateLightCardMaterialInstance = true;
 		}
 
 		if (bAlphaGradientEnable.IsSet() && Actor->AlphaGradient.bEnableAlphaGradient != bAlphaGradientEnable)
 		{
 			Actor->AlphaGradient.bEnableAlphaGradient = bAlphaGradientEnable.GetValue();
-			bNeedsUpdatePolygonTexture = true;
+			bNeedsUpdateLightCardMaterialInstance = true;
 		}
-		
+
 		// Apply Alpha gradient only if Alpha Gradient Enable
 		if (IsAlphaGradientEnabled())
 		{
 			if (StartingAlpha.IsSet() && Actor->AlphaGradient.StartingAlpha != StartingAlpha)
 			{
 				Actor->AlphaGradient.StartingAlpha = StartingAlpha.GetValue();
-				bNeedsUpdatePolygonTexture = true;
+				bNeedsUpdateLightCardMaterialInstance = true;
 			}
 
 			if (EndingAlpha.IsSet() && Actor->AlphaGradient.EndingAlpha != EndingAlpha)
 			{
 				Actor->AlphaGradient.EndingAlpha = EndingAlpha.GetValue();
-				bNeedsUpdatePolygonTexture = true;
+				bNeedsUpdateLightCardMaterialInstance = true;
 			}
 
 			if (GradientAngle.IsSet() && Actor->AlphaGradient.Angle != GradientAngle)
 			{
 				Actor->AlphaGradient.Angle = GradientAngle.GetValue();
-				bNeedsUpdatePolygonTexture = true;
+				bNeedsUpdateLightCardMaterialInstance = true;
 			}
 		}
 
